Updated Data Charts Plugin for new Placement Mode API.  Included temporary icon that will be replaced.

diff --git a/Engine/Plugins/VirtualProduction/DataCharts/Source/DataChartsEditor/Private/DataChartsPlacement.cpp b/Engine/Plugins/VirtualProduction/DataCharts/Source/DataChartsEditor/Private/DataChartsPlacement.cpp
--- a/Engine/Plugins/VirtualProduction/DataCharts/Source/DataChartsEditor/Private/DataChartsPlacement.cpp
+++ b/Engine/Plugins/VirtualProduction/DataCharts/Source/DataChartsEditor/Private/DataChartsPlacement.cpp
@@ -2,6 +2,8 @@
 
 #include "DataChartsPlacement.h"
 #include "IPlacementModeModule.h"
+#include "DataChartsStyle.h"
+
 #include "ActorFactories/ActorFactoryBlueprint.h"
 
 #define LOCTEXT_NAMESPACE "FDataChartsEditorModule"
@@ -10,7 +12,9 @@ void FDataChartsPlacement::RegisterPlacement()
 {
 	IPlacementModeModule& PlacementModeModule = IPlacementModeModule::Get();
 	FName CategoryName = "DataCharts";
-	FPlacementCategoryInfo Info(LOCTEXT("PlacementMode_DataCharts", "Data Charts"), CategoryName, TEXT("PMDataCharts"),45);
+	FPlacementCategoryInfo Info(LOCTEXT("PlacementMode_DataCharts", "Data Charts"), 
+		FSlateIcon(FDataChartsStyle::Get()->GetStyleSetName(), "Icons.VirtualProduction"),
+		CategoryName, TEXT("PMDataCharts"),45);
 	PlacementModeModule.RegisterPlacementCategory(Info);
 
 	TMap<const FName, UBlueprint*> MCharts;
diff --git a/Engine/Plugins/VirtualProduction/DataCharts/Source/DataChartsEditor/Private/DataChartsStyle.cpp b/Engine/Plugins/VirtualProduction/DataCharts/Source/DataChartsEditor/Private/DataChartsStyle.cpp
--- a/Engine/Plugins/VirtualProduction/DataCharts/Source/DataChartsEditor/Private/DataChartsStyle.cpp
+++ b/Engine/Plugins/VirtualProduction/DataCharts/Source/DataChartsEditor/Private/DataChartsStyle.cpp
@@ -5,6 +5,7 @@
 #include "Styling/SlateStyleRegistry.h"
 #include "Styling/SlateTypes.h"
 #include "Interfaces/IPluginManager.h"
+#include "Styling/SlateStyleMacros.h"
 
 TSharedPtr<FSlateStyleSet> FDataChartsStyle::StyleSet;
 
@@ -24,9 +25,12 @@ void FDataChartsStyle::Initialize()
 	FString IconPathBar = IconFolder + TEXT("/Resources/Bar.png");
 	FString IconPathPie = IconFolder + TEXT("/Resources/Pie.png");
 	FString IconPathLine = IconFolder + TEXT("/Resources/Line.png");
+	FString VPIcon = IconFolder + TEXT("/Resources/VirtualProduction.svg");
 	StyleSet->Set("DataCharts.BarIcon", new FSlateImageBrush(IconPathBar, FVector2D(40.0f, 40.0f)));
 	StyleSet->Set("DataCharts.PieIcon", new FSlateImageBrush(IconPathPie, FVector2D(40.0f, 40.0f)));
 	StyleSet->Set("DataCharts.LineIcon", new FSlateImageBrush(IconPathLine, FVector2D(40.0f, 40.0f)));
+	StyleSet->Set("Icons.VirtualProduction", new FSlateVectorImageBrush(VPIcon, FVector2D(16.f, 16.f)));
+
 
 	FSlateStyleRegistry::RegisterSlateStyle(*StyleSet.Get());
 }
